Rename/move file(s) from UncompressedTextureBuildWorker to BaseTextureBuildWorker

diff --git a/Engine/Source/Programs/BaseTextureBuildWorker/BaseTextureBuildWorker.Build.cs b/Engine/Source/Programs/BaseTextureBuildWorker/BaseTextureBuildWorker.Build.cs
new file mode 100644
--- /dev/null
+++ b/Engine/Source/Programs/BaseTextureBuildWorker/BaseTextureBuildWorker.Build.cs
@@ -0,0 +1,25 @@
+// Copyright Epic Games, Inc. All Rights Reserved.
+
+using UnrealBuildTool;
+
+public class UncompressedTextureBuildWorker : ModuleRules
+{
+	public UncompressedTextureBuildWorker(ReadOnlyTargetRules Target) : base(Target)
+	{
+		PrivateIncludePathModuleNames.AddRange(new string[]
+		{
+			"CoreUObject",
+			"Engine",
+			"DerivedDataCache",
+		});
+
+		PrivateDependencyModuleNames.AddRange(new string[]
+		{
+			"Core",
+			"DerivedDataBuildWorker",
+			"TextureBuild",
+			"TextureFormat",
+			"TextureFormatUncompressed",
+		});
+	}
+}
diff --git a/Engine/Source/Programs/BaseTextureBuildWorker/BaseTextureBuildWorker.Target.cs b/Engine/Source/Programs/BaseTextureBuildWorker/BaseTextureBuildWorker.Target.cs
new file mode 100644
--- /dev/null
+++ b/Engine/Source/Programs/BaseTextureBuildWorker/BaseTextureBuildWorker.Target.cs
@@ -0,0 +1,13 @@
+// Copyright Epic Games, Inc. All Rights Reserved.
+
+using UnrealBuildTool;
+using System.Collections.Generic;
+
+[SupportedPlatforms(UnrealPlatformClass.Desktop)]
+public class UncompressedTextureBuildWorkerTarget : DerivedDataBuildWorkerTarget
+{
+	public UncompressedTextureBuildWorkerTarget(TargetInfo Target) : base(Target)
+	{
+		LaunchModuleName = "UncompressedTextureBuildWorker";
+	}
+}
diff --git a/Engine/Source/Programs/BaseTextureBuildWorker/Private/BaseTextureBuildWorker.cpp b/Engine/Source/Programs/BaseTextureBuildWorker/Private/BaseTextureBuildWorker.cpp
new file mode 100644
--- /dev/null
+++ b/Engine/Source/Programs/BaseTextureBuildWorker/Private/BaseTextureBuildWorker.cpp
@@ -0,0 +1,19 @@
+// Copyright Epic Games, Inc. All Rights Reserved.
+
+#include "CoreMinimal.h"
+
+#include "DerivedDataBuildWorkerInterface.h"
+#include "LaunchEngineLoop.h"
+#include "Modules/ModuleManager.h"
+#include "TextureBuildFunction.h"
+
+IMPLEMENT_APPLICATION(UncompressedTextureBuildWorker, "UncompressedTextureBuildWorker");
+
+
+void DerivedDataBuildWorkerInit()
+{
+	static FTextureBuildFunction TextureBuildFunction;
+	TextureBuildFunction.SetName(TEXT("UncompressedTexture"));
+	UE::DerivedData::RegisterWorkerBuildFunction(&TextureBuildFunction);
+}
+
diff --git a/Engine/Source/Programs/UncompressedTextureBuildWorker/Private/UncompressedTextureBuildWorker.cpp b/Engine/Source/Programs/UncompressedTextureBuildWorker/Private/UncompressedTextureBuildWorker.cpp
deleted file mode 100644
diff --git a/Engine/Source/Programs/UncompressedTextureBuildWorker/UncompressedTextureBuildWorker.Build.cs b/Engine/Source/Programs/UncompressedTextureBuildWorker/UncompressedTextureBuildWorker.Build.cs
deleted file mode 100644
diff --git a/Engine/Source/Programs/UncompressedTextureBuildWorker/UncompressedTextureBuildWorker.Target.cs b/Engine/Source/Programs/UncompressedTextureBuildWorker/UncompressedTextureBuildWorker.Target.cs
deleted file mode 100644
