Modified implementation of UncompressedTextureBuildWorker to match the new name of BaseTextureBuildWorker.  Also linked in Oodle texture format module and exposed an OodleTexture build function from BaseTextureBuildWorker.

diff --git a/Engine/Source/Programs/BaseTextureBuildWorker/BaseTextureBuildWorker.Build.cs b/Engine/Source/Programs/BaseTextureBuildWorker/BaseTextureBuildWorker.Build.cs
--- a/Engine/Source/Programs/BaseTextureBuildWorker/BaseTextureBuildWorker.Build.cs
+++ b/Engine/Source/Programs/BaseTextureBuildWorker/BaseTextureBuildWorker.Build.cs
@@ -1,10 +1,13 @@
 // Copyright Epic Games, Inc. All Rights Reserved.
 
+using System;
+using System.Linq;
+using EpicGames.Core;
 using UnrealBuildTool;
 
-public class UncompressedTextureBuildWorker : ModuleRules
+public class BaseTextureBuildWorker : ModuleRules
 {
-	public UncompressedTextureBuildWorker(ReadOnlyTargetRules Target) : base(Target)
+	public BaseTextureBuildWorker(ReadOnlyTargetRules Target) : base(Target)
 	{
 		PrivateIncludePathModuleNames.AddRange(new string[]
 		{
@@ -21,5 +24,22 @@ public class UncompressedTextureBuildWorker : ModuleRules
 			"TextureFormat",
 			"TextureFormatUncompressed",
 		});
+
+		// Determine if TextureFormatOodle is enabled.
+		var TextureFormatOodleUPluginFile = FileReference.Combine(new DirectoryReference(EngineDirectory), "Plugins/Developer/TextureFormatOodle/TextureFormatOodle.uplugin");
+		var TextureFormatOodlePlugin = new PluginInfo(TextureFormatOodleUPluginFile, PluginType.Engine);
+
+		bool bTextureFormatOodlePluginEnabled =
+		Enum.GetValues(typeof(UnrealTargetConfiguration)).Cast<UnrealTargetConfiguration>().Any(config
+				=> Plugins.IsPluginEnabledForTarget(TextureFormatOodlePlugin, null, Target.Platform, config, TargetType.Program));
+
+		if (bTextureFormatOodlePluginEnabled)
+		{
+			PrivateDependencyModuleNames.AddRange(new string[]
+			{
+				"TextureFormatOodle",
+			});
+		}
+
 	}
 }
diff --git a/Engine/Source/Programs/BaseTextureBuildWorker/BaseTextureBuildWorker.Target.cs b/Engine/Source/Programs/BaseTextureBuildWorker/BaseTextureBuildWorker.Target.cs
--- a/Engine/Source/Programs/BaseTextureBuildWorker/BaseTextureBuildWorker.Target.cs
+++ b/Engine/Source/Programs/BaseTextureBuildWorker/BaseTextureBuildWorker.Target.cs
@@ -4,10 +4,10 @@ using UnrealBuildTool;
 using System.Collections.Generic;
 
 [SupportedPlatforms(UnrealPlatformClass.Desktop)]
-public class UncompressedTextureBuildWorkerTarget : DerivedDataBuildWorkerTarget
+public class BaseTextureBuildWorkerTarget : DerivedDataBuildWorkerTarget
 {
-	public UncompressedTextureBuildWorkerTarget(TargetInfo Target) : base(Target)
+	public BaseTextureBuildWorkerTarget(TargetInfo Target) : base(Target)
 	{
-		LaunchModuleName = "UncompressedTextureBuildWorker";
+		LaunchModuleName = "BaseTextureBuildWorker";
 	}
 }
diff --git a/Engine/Source/Programs/BaseTextureBuildWorker/Private/BaseTextureBuildWorker.cpp b/Engine/Source/Programs/BaseTextureBuildWorker/Private/BaseTextureBuildWorker.cpp
--- a/Engine/Source/Programs/BaseTextureBuildWorker/Private/BaseTextureBuildWorker.cpp
+++ b/Engine/Source/Programs/BaseTextureBuildWorker/Private/BaseTextureBuildWorker.cpp
@@ -7,13 +7,16 @@
 #include "Modules/ModuleManager.h"
 #include "TextureBuildFunction.h"
 
-IMPLEMENT_APPLICATION(UncompressedTextureBuildWorker, "UncompressedTextureBuildWorker");
-
+IMPLEMENT_APPLICATION(BaseTextureBuildWorker, "BaseTextureBuildWorker");
 
 void DerivedDataBuildWorkerInit()
 {
-	static FTextureBuildFunction TextureBuildFunction;
-	TextureBuildFunction.SetName(TEXT("UncompressedTexture"));
-	UE::DerivedData::RegisterWorkerBuildFunction(&TextureBuildFunction);
+	static FTextureBuildFunction UncompressedTextureBuildFunction;
+	UncompressedTextureBuildFunction.SetName(TEXT("UncompressedTexture"));
+	UE::DerivedData::RegisterWorkerBuildFunction(&UncompressedTextureBuildFunction);
+
+	static FTextureBuildFunction OodleTextureBuildFunction;
+	OodleTextureBuildFunction.SetName(TEXT("OodleTexture"));
+	UE::DerivedData::RegisterWorkerBuildFunction(&OodleTextureBuildFunction);
 }
 
